Added a Flurry analytics plugin for IOS

diff --git a/Engine/Plugins/Runtime/Analytics/Flurry/IOSFlurry/IOSFlurry.uplugin b/Engine/Plugins/Runtime/Analytics/Flurry/IOSFlurry/IOSFlurry.uplugin
new file mode 100644
--- /dev/null
+++ b/Engine/Plugins/Runtime/Analytics/Flurry/IOSFlurry/IOSFlurry.uplugin
@@ -0,0 +1,23 @@
+{
+	"FileVersion" : 3,
+	
+	"FriendlyName" : "IOS Flurry Analytics Provider",
+	"Version" : 1,
+	"VersionName" : "1.0",
+	"CreatedBy" : "Epic Games, Inc.",
+	"CreatedByURL" : "http://epicgames.com",
+	"DocsURL" : "https://docs.unrealengine.com/latest/INT/Gameplay/Analytics/index.html",
+	"EngineVersion" : "4.8.0",
+	"Description" : "Flurry Analytics Provider for iOS",
+	"Category" : "Analytics.Provider",
+	
+	"Modules" :
+	[
+		{
+			"Name" : "IOSFlurry",
+			"Type" : "Runtime",
+			"LoadingPhase" : "PreDefault",
+			"WhitelistPlatforms" : [ "IOS" ]
+		}
+	]
+}
\ No newline at end of file
diff --git a/Engine/Plugins/Runtime/Analytics/Flurry/IOSFlurry/Source/IOSFlurry/IOSFlurry.Build.cs b/Engine/Plugins/Runtime/Analytics/Flurry/IOSFlurry/Source/IOSFlurry/IOSFlurry.Build.cs
new file mode 100644
--- /dev/null
+++ b/Engine/Plugins/Runtime/Analytics/Flurry/IOSFlurry/Source/IOSFlurry/IOSFlurry.Build.cs
@@ -0,0 +1,63 @@
+// Copyright 1998-2015 Epic Games, Inc. All Rights Reserved.
+
+namespace UnrealBuildTool.Rules
+{
+	public class IOSFlurry : ModuleRules
+	{
+        public IOSFlurry(TargetInfo Target)
+		{
+			BinariesSubFolder = "NotForLicensees";
+
+			PublicDependencyModuleNames.AddRange(
+				new string[]
+				{
+					"Core",
+					// ... add other public dependencies that you statically link with here ...
+				}
+				);
+
+			PrivateDependencyModuleNames.AddRange(
+				new string[]
+				{
+					"Analytics",
+					// ... add private dependencies that you statically link with here ...
+				}
+				);
+
+            PublicIncludePathModuleNames.Add("Analytics");
+
+			PublicFrameworks.AddRange(
+				new string[] {
+					"CoreTelephony",
+					"SystemConfiguration",
+					"UIKit",
+					"Foundation",
+					"CoreGraphics",
+					"MobileCoreServices",
+					"StoreKit",
+					"CFNetwork",
+					"CoreData",
+					"Security",
+					"CoreLocation"
+				});
+
+            bool bHasFlurrySDK =
+                (System.IO.Directory.Exists(System.IO.Path.Combine(UEBuildConfiguration.UEThirdPartySourceDirectory, "Flurry")) &&
+                  System.IO.Directory.Exists(System.IO.Path.Combine(UEBuildConfiguration.UEThirdPartySourceDirectory, "Flurry", "IOS"))) ||
+                (System.IO.Directory.Exists(System.IO.Path.Combine(UEBuildConfiguration.UEThirdPartySourceDirectory, "NotForLicensees")) &&
+                  System.IO.Directory.Exists(System.IO.Path.Combine(UEBuildConfiguration.UEThirdPartySourceDirectory, "NotForLicensees", "Flurry")) &&
+                  System.IO.Directory.Exists(System.IO.Path.Combine(UEBuildConfiguration.UEThirdPartySourceDirectory, "NotForLicensees", "Flurry", "IOS")));
+            if (bHasFlurrySDK)
+            {
+                PublicIncludePaths.Add(UEBuildConfiguration.UEThirdPartySourceDirectory + "NotForLicensees/Flurry/IOS/");
+                PublicAdditionalLibraries.Add(UEBuildConfiguration.UEThirdPartySourceDirectory + "NotForLicensees/Flurry/IOS/libFlurry_6.3.0.a");
+
+                Definitions.Add("WITH_FLURRY=1");
+            }
+            else
+            {
+                Definitions.Add("WITH_FLURRY=0");
+            }
+        }
+	}
+}
\ No newline at end of file
diff --git a/Engine/Plugins/Runtime/Analytics/Flurry/IOSFlurry/Source/IOSFlurry/Private/IOSFlurry.cpp b/Engine/Plugins/Runtime/Analytics/Flurry/IOSFlurry/Source/IOSFlurry/Private/IOSFlurry.cpp
new file mode 100644
--- /dev/null
+++ b/Engine/Plugins/Runtime/Analytics/Flurry/IOSFlurry/Source/IOSFlurry/Private/IOSFlurry.cpp
@@ -0,0 +1,399 @@
+// Copyright 1998-2015 Epic Games, Inc. All Rights Reserved.
+
+#include "IOSFlurryPrivatePCH.h"
+
+DEFINE_LOG_CATEGORY_STATIC(LogAnalytics, Display, All);
+
+IMPLEMENT_MODULE( FAnalyticsIOSFlurry, IOSFlurry )
+
+TSharedPtr<IAnalyticsProvider> FAnalyticsProviderFlurry::Provider;
+
+void FAnalyticsIOSFlurry::StartupModule()
+{
+}
+
+void FAnalyticsIOSFlurry::ShutdownModule()
+{
+	FAnalyticsProviderFlurry::Destroy();
+}
+
+TSharedPtr<IAnalyticsProvider> FAnalyticsIOSFlurry::CreateAnalyticsProvider(const FAnalytics::FProviderConfigurationDelegate& GetConfigValue) const
+{
+	if (GetConfigValue.IsBound())
+	{
+		const FString Key = GetConfigValue.Execute(TEXT("FlurryApiKey"), true);
+		return FAnalyticsProviderFlurry::Create(Key);
+	}
+	else
+	{
+		UE_LOG(LogAnalytics, Warning, TEXT("IOSFlurry::CreateAnalyticsProvider called with an unbound config delegate"));
+	}
+	return nullptr;
+}
+
+// Provider
+
+FAnalyticsProviderFlurry::FAnalyticsProviderFlurry(const FString Key) :
+	ApiKey(Key)
+{
+#if WITH_FLURRY
+
+	// Support more config options here
+
+#else
+	UE_LOG(LogAnalytics, Warning, TEXT("WITH_FLURRY=0. Are you missing the SDK?"));
+#endif
+}
+
+FAnalyticsProviderFlurry::~FAnalyticsProviderFlurry()
+{
+	if (bHasSessionStarted)
+	{
+		EndSession();
+	}
+}
+
+bool FAnalyticsProviderFlurry::StartSession(const TArray<FAnalyticsEventAttribute>& Attributes)
+{
+#if WITH_FLURRY
+	if (!bHasSessionStarted)
+	{
+		NSString* ApiKeyStr = [NSString stringWithFString : ApiKey];
+		const int32 AttrCount = Attributes.Num();
+
+		[Flurry startSession : ApiKeyStr];
+		if (AttrCount > 0)
+		{
+			RecordEvent(TEXT("SessionAttributes"), Attributes);
+		}
+
+		UE_LOG(LogAnalytics, Display, TEXT("IOSFlurry::StartSession(%d attributes)"), AttrCount);
+		bHasSessionStarted = true;
+	}
+	return bHasSessionStarted;
+#else
+	UE_LOG(LogAnalytics, Warning, TEXT("WITH_FLURRY=0. Are you missing the SDK?"));
+	return false;
+#endif
+}
+
+void FAnalyticsProviderFlurry::EndSession()
+{
+#if WITH_FLURRY
+	// Flurry doesn't support ending a session
+	UE_LOG(LogAnalytics, Display, TEXT("IOSFlurry::EndSession - ignoring call"));
+#else
+	UE_LOG(LogAnalytics, Warning, TEXT("WITH_FLURRY=0. Are you missing the SDK?"));
+#endif
+}
+
+void FAnalyticsProviderFlurry::FlushEvents()
+{
+#if WITH_FLURRY
+	// Flurry doesn't support flushing a session
+	UE_LOG(LogAnalytics, Display, TEXT("IOSFlurry::FlushEvents - ignoring call"));
+#else
+	UE_LOG(LogAnalytics, Warning, TEXT("WITH_FLURRY=0. Are you missing the SDK?"));
+#endif
+}
+
+void FAnalyticsProviderFlurry::SetUserID(const FString& InUserID)
+{
+#if WITH_FLURRY
+	// Cache the string in case GetUserId is called
+	UserId = InUserID;
+	NSString* UserIdStr = [NSString stringWithFString : UserId];
+
+	[Flurry setUserID:UserIdStr];
+
+	UE_LOG(LogAnalytics, Display, TEXT("IOSFlurry::SetUserID(%s)"), *UserId);
+#else
+	UE_LOG(LogAnalytics, Warning, TEXT("WITH_FLURRY=0. Are you missing the SDK?"));
+#endif
+}
+
+FString FAnalyticsProviderFlurry::GetUserID() const
+{
+#if WITH_FLURRY
+	UE_LOG(LogAnalytics, Display, TEXT("IOSFlurry::GetUserID - returning cached id '%s'"), *UserId);
+
+	return UserId;
+#else
+	UE_LOG(LogAnalytics, Warning, TEXT("WITH_FLURRY=0. Are you missing the SDK?"));
+	return FString();
+#endif
+}
+
+void FAnalyticsProviderFlurry::SetGender(const FString& InGender)
+{
+#if WITH_FLURRY
+	NSString* ConvertedStr = [NSString stringWithFString : InGender];
+
+	[Flurry setGender:ConvertedStr];
+
+	UE_LOG(LogAnalytics, Display, TEXT("IOSFlurry::SetGender(%s)"), *InGender);
+#else
+	UE_LOG(LogAnalytics, Warning, TEXT("WITH_FLURRY=0. Are you missing the SDK?"));
+#endif
+}
+
+void FAnalyticsProviderFlurry::SetAge(const int32 InAge)
+{
+#if WITH_FLURRY
+	[Flurry setAge:InAge];
+
+	UE_LOG(LogAnalytics, Display, TEXT("IOSFlurry::SetAge(%d)"), InAge);
+#else
+	UE_LOG(LogAnalytics, Warning, TEXT("WITH_FLURRY=0. Are you missing the SDK?"));
+#endif
+}
+
+void FAnalyticsProviderFlurry::SetLocation(const FString& InLocation)
+{
+#if WITH_FLURRY
+	FString Lat, Long;
+	InLocation.Split(TEXT(","), &Lat, &Long);
+
+	double Latitude = FCString::Atod(*Lat);
+	double Longitude = FCString::Atod(*Long);
+
+	[Flurry setLatitude:Latitude longitude:Longitude horizontalAccuracy:0.0 verticalAccuracy:0.0];
+
+	UE_LOG(LogAnalytics, Display, TEXT("Parsed \"lat, long\" string in IOSFlurry::SetLocation(%s) as \"%f, %f\""), *InLocation, Latitude, Longitude);
+#else
+	UE_LOG(LogAnalytics, Warning, TEXT("WITH_FLURRY=0. Are you missing the SDK?"));
+#endif
+}
+
+FString FAnalyticsProviderFlurry::GetSessionID() const
+{
+#if WITH_FLURRY
+	NSString* Id = [Flurry getSessionID];
+	FString ConvertedId(Id);
+
+	UE_LOG(LogAnalytics, Display, TEXT("IOSFlurry::GetSessionID - returning the id as '%s'"), *ConvertedId);
+
+	return ConvertedId;
+#else
+	UE_LOG(LogAnalytics, Warning, TEXT("WITH_FLURRY=0. Are you missing the SDK?"));
+	return FString();
+#endif
+}
+
+bool FAnalyticsProviderFlurry::SetSessionID(const FString& InSessionID)
+{
+#if WITH_FLURRY
+	// Ignored
+	UE_LOG(LogAnalytics, Display, TEXT("IOSFlurry::SetSessionID - ignoring call"));
+	return false;
+#else
+	UE_LOG(LogAnalytics, Warning, TEXT("WITH_FLURRY=0. Are you missing the SDK?"));
+	return false;
+#endif
+}
+
+void FAnalyticsProviderFlurry::RecordEvent(const FString& EventName, const TArray<FAnalyticsEventAttribute>& Attributes)
+{
+#if WITH_FLURRY
+	if (EventName.Len() > 0)
+	{
+		NSString* ConvertedEventName = [NSString stringWithFString : EventName];
+		const int32 AttrCount = Attributes.Num();
+		if (AttrCount > 0)
+		{
+			// Convert the event attributes into a dictionary object
+			NSDictionary* AttributesDict = [NSMutableDictionary dictionaryWithCapacity:AttrCount];
+			for	(auto Attr : Attributes)
+			{
+				NSString* AttrName = [NSString stringWithFString : Attr.AttrName];
+				NSString* AttrValue = [NSString stringWithFString : Attr.AttrValue];
+				[AttributesDict setValue:AttrValue forKey:AttrName];
+			}
+			[Flurry logEvent:ConvertedEventName withParameters:AttributesDict];
+		}
+		else
+		{
+			// Log an event with no payload
+			[Flurry logEvent:ConvertedEventName];
+		}
+		UE_LOG(LogAnalytics, Display, TEXT("IOSFlurry::RecordEvent('%s', %d attributes)"), *EventName, AttrCount);
+	}
+#else
+	UE_LOG(LogAnalytics, Warning, TEXT("WITH_FLURRY=0. Are you missing the SDK?"));
+#endif
+}
+
+void FAnalyticsProviderFlurry::RecordItemPurchase(const FString& ItemId, const FString& Currency, int PerItemCost, int ItemQuantity)
+{
+#if WITH_FLURRY
+	NSString* EventName = @"Item Purchase";
+	// Build the dictionary
+	NSDictionary* AttributesDict = [NSMutableDictionary dictionaryWithCapacity:4];
+	[AttributesDict setValue:[NSString stringWithFString:ItemId] forKey:@"ItemId"];
+	[AttributesDict setValue:[NSString stringWithFString:Currency] forKey:@"Currency"];
+	[AttributesDict setValue:[NSString stringWithFormat:@"%d", PerItemCost] forKey:@"PerItemCost"];
+	[AttributesDict setValue:[NSString stringWithFormat:@"%d", ItemQuantity] forKey:@"ItemQuantity"];
+	// Send the event
+	[Flurry logEvent:EventName withParameters:AttributesDict];
+
+	UE_LOG(LogAnalytics, Display, TEXT("IOSFlurry::RecordItemPurchase('%s', '%s', %d, %d)"), *ItemId, *Currency, PerItemCost, ItemQuantity);
+#else
+	UE_LOG(LogAnalytics, Warning, TEXT("WITH_FLURRY=0. Are you missing the SDK?"));
+#endif
+}
+
+void FAnalyticsProviderFlurry::RecordCurrencyPurchase(const FString& GameCurrencyType, int GameCurrencyAmount, const FString& RealCurrencyType, float RealMoneyCost, const FString& PaymentProvider)
+{
+#if WITH_FLURRY
+	NSString* EventName = @"Currency Purchase";
+	// Build the dictionary
+	NSDictionary* AttributesDict = [NSMutableDictionary dictionaryWithCapacity:5];
+	[AttributesDict setValue:[NSString stringWithFString:GameCurrencyType] forKey:@"GameCurrencyType"];
+	[AttributesDict setValue:[NSString stringWithFormat:@"%d", GameCurrencyAmount] forKey:@"GameCurrencyAmount"];
+	[AttributesDict setValue:[NSString stringWithFString:RealCurrencyType] forKey:@"RealCurrencyType"];
+	[AttributesDict setValue:[NSString stringWithFormat:@"%.02f", RealMoneyCost] forKey:@"RealMoneyCost"];
+	[AttributesDict setValue:[NSString stringWithFString:PaymentProvider] forKey:@"PaymentProvider"];
+	// Send the event
+	[Flurry logEvent:EventName withParameters:AttributesDict];
+
+	UE_LOG(LogAnalytics, Display, TEXT("IOSFlurry::RecordCurrencyPurchase('%s', %d, '%s', %.02f, %s)"), *GameCurrencyType, GameCurrencyAmount, *RealCurrencyType, RealMoneyCost, *PaymentProvider);
+#else
+	UE_LOG(LogAnalytics, Warning, TEXT("WITH_FLURRY=0. Are you missing the SDK?"));
+#endif
+}
+
+void FAnalyticsProviderFlurry::RecordCurrencyGiven(const FString& GameCurrencyType, int GameCurrencyAmount)
+{
+#if WITH_FLURRY
+	NSString* EventName = @"Currency Given";
+	// Build the dictionary
+	NSDictionary* AttributesDict = [NSMutableDictionary dictionaryWithCapacity:2];
+	[AttributesDict setValue:[NSString stringWithFString:GameCurrencyType] forKey:@"GameCurrencyType"];
+	[AttributesDict setValue:[NSString stringWithFormat:@"%d", GameCurrencyAmount] forKey:@"GameCurrencyAmount"];
+	// Send the event
+	[Flurry logEvent:EventName withParameters:AttributesDict];
+
+	UE_LOG(LogAnalytics, Display, TEXT("IOSFlurry::RecordCurrencyGiven('%s', %d)"), *GameCurrencyType, GameCurrencyAmount);
+#else
+	UE_LOG(LogAnalytics, Warning, TEXT("WITH_FLURRY=0. Are you missing the SDK?"));
+#endif
+}
+
+void FAnalyticsProviderFlurry::RecordItemPurchase(const FString& ItemId, int ItemQuantity, const TArray<FAnalyticsEventAttribute>& EventAttrs)
+{
+#if WITH_FLURRY
+	NSString* EventName = @"Item Purchase";
+	// Build the dictionary
+	int32 DictSize = EventAttrs.Num() + 2;
+	NSDictionary* AttributesDict = [NSMutableDictionary dictionaryWithCapacity:DictSize];
+	[AttributesDict setValue:[NSString stringWithFString:ItemId] forKey:@"ItemId"];
+	[AttributesDict setValue:[NSString stringWithFormat:@"%d", ItemQuantity] forKey:@"Quantity"];
+	for	(auto Attr : EventAttrs)
+	{
+		NSString* AttrName = [NSString stringWithFString : Attr.AttrName];
+		NSString* AttrValue = [NSString stringWithFString : Attr.AttrValue];
+		[AttributesDict setValue:AttrValue forKey:AttrName];
+	}
+	// Send the event
+	[Flurry logEvent:EventName withParameters:AttributesDict];
+
+	UE_LOG(LogAnalytics, Display, TEXT("IOSFlurry::RecordItemPurchase('%s', %d, %d)"), *ItemId, ItemQuantity, EventAttrs.Num());
+#else
+	UE_LOG(LogAnalytics, Warning, TEXT("WITH_FLURRY=0. Are you missing the SDK?"));
+#endif
+}
+
+void FAnalyticsProviderFlurry::RecordCurrencyPurchase(const FString& GameCurrencyType, int GameCurrencyAmount, const TArray<FAnalyticsEventAttribute>& EventAttrs)
+{
+#if WITH_FLURRY
+	NSString* EventName = @"Currency Purchase";
+	// Build the dictionary
+	int32 DictSize = EventAttrs.Num() + 2;
+	NSDictionary* AttributesDict = [NSMutableDictionary dictionaryWithCapacity:DictSize];
+	[AttributesDict setValue:[NSString stringWithFString:GameCurrencyType] forKey:@"GameCurrencyType"];
+	[AttributesDict setValue:[NSString stringWithFormat:@"%d", GameCurrencyAmount] forKey:@"GameCurrencyAmount"];
+	for	(auto Attr : EventAttrs)
+	{
+		NSString* AttrName = [NSString stringWithFString : Attr.AttrName];
+		NSString* AttrValue = [NSString stringWithFString : Attr.AttrValue];
+		[AttributesDict setValue:AttrValue forKey:AttrName];
+	}
+	// Send the event
+	[Flurry logEvent:EventName withParameters:AttributesDict];
+
+	UE_LOG(LogAnalytics, Display, TEXT("IOSFlurry::RecordCurrencyPurchase('%s', %d, %d)"), *GameCurrencyType, GameCurrencyAmount, EventAttrs.Num());
+#else
+	UE_LOG(LogAnalytics, Warning, TEXT("WITH_FLURRY=0. Are you missing the SDK?"));
+#endif
+}
+
+void FAnalyticsProviderFlurry::RecordCurrencyGiven(const FString& GameCurrencyType, int GameCurrencyAmount, const TArray<FAnalyticsEventAttribute>& EventAttrs)
+{
+#if WITH_FLURRY
+	NSString* EventName = @"Currency Given";
+	// Build the dictionary
+	int32 DictSize = EventAttrs.Num() + 2;
+	NSDictionary* AttributesDict = [NSMutableDictionary dictionaryWithCapacity:DictSize];
+	[AttributesDict setValue:[NSString stringWithFString:GameCurrencyType] forKey:@"GameCurrencyType"];
+	[AttributesDict setValue:[NSString stringWithFormat:@"%d", GameCurrencyAmount] forKey:@"GameCurrencyAmount"];
+	for	(auto Attr : EventAttrs)
+	{
+		NSString* AttrName = [NSString stringWithFString : Attr.AttrName];
+		NSString* AttrValue = [NSString stringWithFString : Attr.AttrValue];
+		[AttributesDict setValue:AttrValue forKey:AttrName];
+	}
+	// Send the event
+	[Flurry logEvent:EventName withParameters:AttributesDict];
+
+	UE_LOG(LogAnalytics, Display, TEXT("IOSFlurry::RecordCurrencyGiven('%s', %d, %d)"), *GameCurrencyType, GameCurrencyAmount, EventAttrs.Num());
+#else
+	UE_LOG(LogAnalytics, Warning, TEXT("WITH_FLURRY=0. Are you missing the SDK?"));
+#endif
+}
+
+void FAnalyticsProviderFlurry::RecordError(const FString& Error, const TArray<FAnalyticsEventAttribute>& EventAttrs)
+{
+#if WITH_FLURRY
+	NSString* EventName = @"Error";
+	// Build the dictionary
+	int32 DictSize = EventAttrs.Num() + 1;
+	NSDictionary* AttributesDict = [NSMutableDictionary dictionaryWithCapacity:DictSize];
+	[AttributesDict setValue:[NSString stringWithFString:Error] forKey:@"Error"];
+	for	(auto Attr : EventAttrs)
+	{
+		NSString* AttrName = [NSString stringWithFString : Attr.AttrName];
+		NSString* AttrValue = [NSString stringWithFString : Attr.AttrValue];
+		[AttributesDict setValue:AttrValue forKey:AttrName];
+	}
+	// Send the event
+	[Flurry logEvent:EventName withParameters:AttributesDict];
+
+	UE_LOG(LogAnalytics, Display, TEXT("IOSFlurry::RecordError('%s', %d)"), *Error, EventAttrs.Num());
+#else
+	UE_LOG(LogAnalytics, Warning, TEXT("WITH_FLURRY=0. Are you missing the SDK?"));
+#endif
+}
+
+void FAnalyticsProviderFlurry::RecordProgress(const FString& ProgressType, const FString& ProgressHierarchy, const TArray<FAnalyticsEventAttribute>& EventAttrs)
+{
+#if WITH_FLURRY
+	NSString* EventName = @"Progress";
+	// Build the dictionary
+	int32 DictSize = EventAttrs.Num() + 2;
+	NSDictionary* AttributesDict = [NSMutableDictionary dictionaryWithCapacity:DictSize];
+	[AttributesDict setValue:[NSString stringWithFString:ProgressType] forKey:@"ProgressType"];
+	[AttributesDict setValue:[NSString stringWithFString:ProgressHierarchy] forKey:@"ProgressHierarchy"];
+	for	(auto Attr : EventAttrs)
+	{
+		NSString* AttrName = [NSString stringWithFString : Attr.AttrName];
+		NSString* AttrValue = [NSString stringWithFString : Attr.AttrValue];
+		[AttributesDict setValue:AttrValue forKey:AttrName];
+	}
+	// Send the event
+	[Flurry logEvent:EventName withParameters:AttributesDict];
+
+	UE_LOG(LogAnalytics, Display, TEXT("IOSFlurry::RecordProgress('%s', %s, %d)"), *ProgressType, *ProgressHierarchy, EventAttrs.Num());
+#else
+	UE_LOG(LogAnalytics, Warning, TEXT("WITH_FLURRY=0. Are you missing the SDK?"));
+#endif
+}
diff --git a/Engine/Plugins/Runtime/Analytics/Flurry/IOSFlurry/Source/IOSFlurry/Private/IOSFlurryPrivatePCH.h b/Engine/Plugins/Runtime/Analytics/Flurry/IOSFlurry/Source/IOSFlurry/Private/IOSFlurryPrivatePCH.h
new file mode 100644
--- /dev/null
+++ b/Engine/Plugins/Runtime/Analytics/Flurry/IOSFlurry/Source/IOSFlurry/Private/IOSFlurryPrivatePCH.h
@@ -0,0 +1,15 @@
+// Copyright 1998-2015 Epic Games, Inc. All Rights Reserved.
+
+#include "Core.h"
+#include "IAnalyticsProviderModule.h"
+#include "IAnalyticsProvider.h"
+
+#include "IOSFlurry.h"
+#include "IOSFlurryProvider.h"
+
+#if WITH_FLURRY
+#import "Flurry.h"
+#endif
+
+// You should place include statements to your module's private header files here.  You only need to
+// add includes for headers that are used in most of your module's source files though.
diff --git a/Engine/Plugins/Runtime/Analytics/Flurry/IOSFlurry/Source/IOSFlurry/Private/IOSFlurryProvider.h b/Engine/Plugins/Runtime/Analytics/Flurry/IOSFlurry/Source/IOSFlurry/Private/IOSFlurryProvider.h
new file mode 100644
--- /dev/null
+++ b/Engine/Plugins/Runtime/Analytics/Flurry/IOSFlurry/Source/IOSFlurry/Private/IOSFlurryProvider.h
@@ -0,0 +1,60 @@
+// Copyright 1998-2015 Epic Games, Inc. All Rights Reserved.
+
+#pragma once
+
+class FAnalyticsProviderFlurry :
+	public IAnalyticsProvider
+{
+	/** The API key given to you by Flurry */
+	FString ApiKey;
+	/** Tracks whether we need to start the session or restart it */
+	bool bHasSessionStarted;
+	/** Cached user id */
+	FString UserId;
+
+	/** Singleton for analytics */
+	static TSharedPtr<IAnalyticsProvider> Provider;
+	FAnalyticsProviderFlurry(const FString Key);
+
+public:
+	static TSharedPtr<IAnalyticsProvider> Create(const FString Key)
+	{
+		if (!Provider.IsValid())
+		{
+			Provider = TSharedPtr<IAnalyticsProvider>(new FAnalyticsProviderFlurry(Key));
+		}
+		return Provider;
+	}
+	static void Destroy()
+	{
+		Provider.Reset();
+	}
+
+	virtual ~FAnalyticsProviderFlurry();
+
+	virtual bool StartSession(const TArray<FAnalyticsEventAttribute>& Attributes) override;
+	virtual void EndSession() override;
+	virtual void FlushEvents() override;
+
+	virtual void SetUserID(const FString& InUserID) override;
+	virtual FString GetUserID() const override;
+
+	virtual FString GetSessionID() const override;
+	virtual bool SetSessionID(const FString& InSessionID) override;
+
+	virtual void RecordEvent(const FString& EventName, const TArray<FAnalyticsEventAttribute>& Attributes) override;
+
+	virtual void RecordItemPurchase(const FString& ItemId, const FString& Currency, int PerItemCost, int ItemQuantity) override;
+	virtual void RecordCurrencyPurchase(const FString& GameCurrencyType, int GameCurrencyAmount, const FString& RealCurrencyType, float RealMoneyCost, const FString& PaymentProvider) override;
+	virtual void RecordCurrencyGiven(const FString& GameCurrencyType, int GameCurrencyAmount) override;
+
+	virtual void SetGender(const FString& InGender) override;
+	virtual void SetLocation(const FString& InLocation) override;
+	virtual void SetAge(const int32 InAge) override;
+
+	virtual void RecordItemPurchase(const FString& ItemId, int ItemQuantity, const TArray<FAnalyticsEventAttribute>& EventAttrs) override;
+	virtual void RecordCurrencyPurchase(const FString& GameCurrencyType, int GameCurrencyAmount, const TArray<FAnalyticsEventAttribute>& EventAttrs) override;
+	virtual void RecordCurrencyGiven(const FString& GameCurrencyType, int GameCurrencyAmount, const TArray<FAnalyticsEventAttribute>& EventAttrs) override;
+	virtual void RecordError(const FString& Error, const TArray<FAnalyticsEventAttribute>& EventAttrs) override;
+	virtual void RecordProgress(const FString& ProgressType, const FString& ProgressHierarchy, const TArray<FAnalyticsEventAttribute>& EventAttrs) override;
+};
diff --git a/Engine/Plugins/Runtime/Analytics/Flurry/IOSFlurry/Source/IOSFlurry/Public/IOSFlurry.h b/Engine/Plugins/Runtime/Analytics/Flurry/IOSFlurry/Source/IOSFlurry/Public/IOSFlurry.h
new file mode 100644
--- /dev/null
+++ b/Engine/Plugins/Runtime/Analytics/Flurry/IOSFlurry/Source/IOSFlurry/Public/IOSFlurry.h
@@ -0,0 +1,49 @@
+// Copyright 1998-2015 Epic Games, Inc. All Rights Reserved.
+
+#pragma once
+
+#include "Runtime/Analytics/Analytics/Public/Interfaces/IAnalyticsProviderModule.h"
+#include "Core.h"
+
+class IAnalyticsProvider;
+
+/**
+ * The public interface to this module
+ */
+class FAnalyticsIOSFlurry :
+	public IAnalyticsProviderModule
+{
+	/** Singleton for analytics */
+	TSharedPtr<IAnalyticsProvider> Provider;
+
+	//--------------------------------------------------------------------------
+	// Module functionality
+	//--------------------------------------------------------------------------
+public:
+	/**
+	 * Singleton-like access to this module's interface.  This is just for convenience!
+	 * Beware of calling this during the shutdown phase, though.  Your module might have been unloaded already.
+	 *
+	 * @return Returns singleton instance, loading the module on demand if needed
+	 */
+	static inline FAnalyticsIOSFlurry& Get()
+	{
+		return FModuleManager::LoadModuleChecked< FAnalyticsIOSFlurry >( "IOSFlurry" );
+	}
+
+	//--------------------------------------------------------------------------
+	// provider factory functions
+	//--------------------------------------------------------------------------
+public:
+	/**
+	 * IAnalyticsProviderModule interface.
+	 * Creates the analytics provider given a configuration delegate.
+	 * The keys required exactly match the field names in the Config object. 
+	 */
+	virtual TSharedPtr<IAnalyticsProvider> CreateAnalyticsProvider(const FAnalytics::FProviderConfigurationDelegate& GetConfigValue) const override;
+	
+private:
+	virtual void StartupModule() override;
+	virtual void ShutdownModule() override;
+};
+
