Update plugin category.

diff --git a/Engine/Plugins/Runtime/Analytics/Flurry/IOSFlurry/IOSFlurry.uplugin b/Engine/Plugins/Runtime/Analytics/Flurry/IOSFlurry/IOSFlurry.uplugin
--- a/Engine/Plugins/Runtime/Analytics/Flurry/IOSFlurry/IOSFlurry.uplugin
+++ b/Engine/Plugins/Runtime/Analytics/Flurry/IOSFlurry/IOSFlurry.uplugin
@@ -9,7 +9,7 @@
 	"DocsURL" : "https://docs.unrealengine.com/latest/INT/Gameplay/Analytics/index.html",
 	"EngineVersion" : "4.8.0",
 	"Description" : "Flurry Analytics Provider for iOS",
-	"Category" : "Analytics.Provider",
+	"Category" : "Analytics",
 	
 	"Modules" :
 	[
